Fix for build error in ITD plugin

diff --git a/Engine/Plugins/Runtime/Spatialization/Source/Spatialization/Private/ITDSpatializer.cpp b/Engine/Plugins/Runtime/Spatialization/Source/Spatialization/Private/ITDSpatializer.cpp
--- a/Engine/Plugins/Runtime/Spatialization/Source/Spatialization/Private/ITDSpatializer.cpp
+++ b/Engine/Plugins/Runtime/Spatialization/Source/Spatialization/Private/ITDSpatializer.cpp
@@ -105,8 +105,8 @@ void FSourceSpatializer::ProcessSource(const FAudioPluginSourceInputData& InputD
 	{
 		for (int32 InChannelIndex = 0; InChannelIndex < NumInputChannels; InChannelIndex++)
 		{
-			OutAudio[OutFrameIndex] += LeftDelays[InChannelIndex].ProcessAudioSample(InAudio[InFrameIndex + InChannelIndex]) * LeftGain.GetValue();
-			OutAudio[OutFrameIndex + 1] += RightDelays[InChannelIndex].ProcessAudioSample(InAudio[InFrameIndex + InChannelIndex]) * RightGain.GetValue();
+			OutAudio[OutFrameIndex] += LeftDelays[InChannelIndex].ProcessAudioSample(InAudio[InFrameIndex + InChannelIndex]) * LeftGain.GetNextValue();
+			OutAudio[OutFrameIndex + 1] += RightDelays[InChannelIndex].ProcessAudioSample(InAudio[InFrameIndex + InChannelIndex]) * RightGain.GetNextValue();
 		}
 
 		OutFrameIndex += NumOutputChannels;
