Audio Plugin Telemetry Part 4
Spatialization, Soundfields, SoundModule, and SoundUtilities

diff --git a/Engine/Plugins/Runtime/SoundFields/Source/SoundFields/Private/SoundFields.cpp b/Engine/Plugins/Runtime/SoundFields/Source/SoundFields/Private/SoundFields.cpp
--- a/Engine/Plugins/Runtime/SoundFields/Source/SoundFields/Private/SoundFields.cpp
+++ b/Engine/Plugins/Runtime/SoundFields/Source/SoundFields/Private/SoundFields.cpp
@@ -5,6 +5,7 @@
 #include "DSP/AlignedBuffer.h"
 #include "DSP/Dsp.h"
 #include "DSP/FloatArrayMath.h"
+#include "EngineAnalytics.h"
 #include "SoundFieldRendering.h"
 #include "HAL/IConsoleManager.h"
 
@@ -251,6 +252,10 @@ TUniquePtr<ISoundfieldTranscodeStream> FAmbisonicsSoundfieldFormat::CreateTransc
 
 TUniquePtr<ISoundfieldMixerStream> FAmbisonicsSoundfieldFormat::CreateMixerStream(const ISoundfieldEncodingSettingsProxy& InitialSettings)
 {
+	if (FEngineAnalytics::IsAvailable())
+	{
+		FEngineAnalytics::GetProvider().RecordEvent("Audio.Usage.SoundFields.MixerStreamCreated");
+	}
 	return TUniquePtr<ISoundfieldMixerStream>(new FAmbisonicsMixer());
 }
 
diff --git a/Engine/Plugins/Runtime/SoundMod/Source/SoundModImporter/Private/SoundModImporterFactory.cpp b/Engine/Plugins/Runtime/SoundMod/Source/SoundModImporter/Private/SoundModImporterFactory.cpp
--- a/Engine/Plugins/Runtime/SoundMod/Source/SoundModImporter/Private/SoundModImporterFactory.cpp
+++ b/Engine/Plugins/Runtime/SoundMod/Source/SoundModImporter/Private/SoundModImporterFactory.cpp
@@ -1,6 +1,7 @@
 // Copyright Epic Games, Inc. All Rights Reserved.
 
 #include "SoundModImporterFactory.h"
+#include "EngineAnalytics.h"
 #include "EngineGlobals.h"
 #include "Editor.h"
 #include "SoundMod.h"
@@ -121,5 +122,10 @@ UObject* USoundModImporterFactory::FactoryCreateBinary
 		ComponentsToRestart[ComponentIndex]->Play();
 	}
 
+	if (FEngineAnalytics::IsAvailable())
+	{
+		FEngineAnalytics::GetProvider().RecordEvent("Audio.Usage.SoundModule.BinaryCreated");
+	}
+
 	return Sound;
 }
diff --git a/Engine/Plugins/Runtime/SoundUtilities/Source/SoundUtilitiesEditor/Private/SoundSimpleFactory.cpp b/Engine/Plugins/Runtime/SoundUtilities/Source/SoundUtilitiesEditor/Private/SoundSimpleFactory.cpp
--- a/Engine/Plugins/Runtime/SoundUtilities/Source/SoundUtilitiesEditor/Private/SoundSimpleFactory.cpp
+++ b/Engine/Plugins/Runtime/SoundUtilities/Source/SoundUtilitiesEditor/Private/SoundSimpleFactory.cpp
@@ -1,6 +1,7 @@
 // Copyright Epic Games, Inc. All Rights Reserved.
 
 #include "SoundSimpleFactory.h"
+#include "EngineAnalytics.h"
 #include "SoundSimple.h"
 
 USoundSimpleFactory::USoundSimpleFactory(const FObjectInitializer& ObjectInitializer)
@@ -36,5 +37,10 @@ UObject* USoundSimpleFactory::FactoryCreateNew(UClass* Class, UObject* InParent,
 	// Write out the duration to be the longest duration
 	SoundSimple->Duration = Duration;
 
+	if (FEngineAnalytics::IsAvailable())
+	{
+		FEngineAnalytics::GetProvider().RecordEvent("Audio.Usage.SoundUtilities.SimpleSoundCreated");
+	}
+
 	return SoundSimple;
 }
diff --git a/Engine/Plugins/Runtime/Spatialization/Source/SpatializationEditor/Private/ITDSpatializationSourceSettingsFactory.cpp b/Engine/Plugins/Runtime/Spatialization/Source/SpatializationEditor/Private/ITDSpatializationSourceSettingsFactory.cpp
--- a/Engine/Plugins/Runtime/Spatialization/Source/SpatializationEditor/Private/ITDSpatializationSourceSettingsFactory.cpp
+++ b/Engine/Plugins/Runtime/Spatialization/Source/SpatializationEditor/Private/ITDSpatializationSourceSettingsFactory.cpp
@@ -1,6 +1,7 @@
 // Copyright Epic Games, Inc. All Rights Reserved.
 
 #include "ITDSpatializationSourceSettingsFactory.h"
+#include "EngineAnalytics.h"
 #include "ITDSpatializationSourceSettings.h"
 
 FText FAssetTypeActions_ITDSpatializationSettings::GetName() const
@@ -45,6 +46,10 @@ UITDSpatializationSettingsFactory::UITDSpatializationSettingsFactory(const FObje
 UObject* UITDSpatializationSettingsFactory::FactoryCreateNew(UClass* Class, UObject* InParent, FName InName, EObjectFlags Flags,
 	UObject* Context, FFeedbackContext* Warn)
 {
+	if (FEngineAnalytics::IsAvailable())
+	{
+		FEngineAnalytics::GetProvider().RecordEvent("Audio.Usage.Spatialization.SettingsCreated");
+	}
 	return NewObject<UITDSpatializationSourceSettings>(InParent, InName, Flags);
 }
 
