Updated Spatialization to not compile with optimizations off so it matches other modules that use Shared PCHs. This speeds up compile times and removes the need to have a 2.03 GB shared Engine PCH.

diff --git a/Engine/Plugins/Runtime/Spatialization/Source/Spatialization/Spatialization.Build.cs b/Engine/Plugins/Runtime/Spatialization/Source/Spatialization/Spatialization.Build.cs
--- a/Engine/Plugins/Runtime/Spatialization/Source/Spatialization/Spatialization.Build.cs
+++ b/Engine/Plugins/Runtime/Spatialization/Source/Spatialization/Spatialization.Build.cs
@@ -8,7 +8,6 @@ public class Spatialization : ModuleRules
 	public Spatialization(ReadOnlyTargetRules Target) : base(Target)
 	{
 		PCHUsage = ModuleRules.PCHUsageMode.UseExplicitOrSharedPCHs;
-        OptimizeCode = CodeOptimization.Never;
 
         PublicDependencyModuleNames.AddRange(
 			new string[]
